[Feature] CmdLink.exe that can run commands on an active UE process

diff --git a/Engine/Plugins/CmdLinkServer/CmdLinkServer.uplugin b/Engine/Plugins/CmdLinkServer/CmdLinkServer.uplugin
new file mode 100644
--- /dev/null
+++ b/Engine/Plugins/CmdLinkServer/CmdLinkServer.uplugin
@@ -0,0 +1,27 @@
+﻿{
+	"FileVersion": 3,
+	"Version": 1,
+	"VersionName": "1.0",
+	"FriendlyName": "Command Link Server",
+	"Description": "Listens to and runs console commands from CmdLink",
+	"Category": "Other",
+	"CreatedBy": "Epic Games, Inc.",
+	"CreatedByURL": "http://epicgames.com",
+	"DocsURL": "",
+	"EnabledByDefault": true,
+	"MarketplaceURL": "",
+	"CanContainContent": false,
+	"IsBetaVersion": false,
+	"IsExperimentalVersion": false,
+	"Installed": false,
+	"Modules": [
+		{
+			"Name": "CmdLinkServer",
+			"Type": "Editor",
+			"LoadingPhase": "PostDefault",
+			"PlatformAllowList": [
+				"Win64"
+			]
+		}
+	]
+}
\ No newline at end of file
diff --git a/Engine/Plugins/CmdLinkServer/Source/CmdLinkServer/CmdLinkServer.Build.cs b/Engine/Plugins/CmdLinkServer/Source/CmdLinkServer/CmdLinkServer.Build.cs
new file mode 100644
--- /dev/null
+++ b/Engine/Plugins/CmdLinkServer/Source/CmdLinkServer/CmdLinkServer.Build.cs
@@ -0,0 +1,24 @@
+﻿// Copyright Epic Games, Inc. All Rights Reserved.
+
+using UnrealBuildTool;
+
+public class CmdLinkServer : ModuleRules
+{
+	public CmdLinkServer(ReadOnlyTargetRules Target) : base(Target)
+	{
+		PCHUsage = ModuleRules.PCHUsageMode.UseExplicitOrSharedPCHs;
+			
+		PublicDependencyModuleNames.AddRange( new string[] 
+		{
+			"Core",
+		});
+			
+		
+		PrivateDependencyModuleNames.AddRange( new string[] 
+		{
+			"CoreUObject",
+			"Engine",
+			"UnrealEd",
+		});
+	}
+}
diff --git a/Engine/Plugins/CmdLinkServer/Source/CmdLinkServer/Private/CmdLinkServer.cpp b/Engine/Plugins/CmdLinkServer/Source/CmdLinkServer/Private/CmdLinkServer.cpp
new file mode 100644
--- /dev/null
+++ b/Engine/Plugins/CmdLinkServer/Source/CmdLinkServer/Private/CmdLinkServer.cpp
@@ -0,0 +1,499 @@
+﻿// Copyright Epic Games, Inc. All Rights Reserved.
+
+#include "CmdLinkServer.h"
+
+#include "Editor.h"
+#include "Async/Async.h"
+#include "Misc/CString.h"
+#include "Features/IModularFeatures.h"
+
+#define LOCTEXT_NAMESPACE "FCmdLinkServerModule"
+
+
+namespace UE::CmdLink
+{
+	bool bEnabled = false;
+	void OnCmdLinkEnabledChanged(IConsoleVariable* Var);
+	FAutoConsoleVariableRef CVarEnableCmdLink(
+		TEXT("console.CmdLink.enable"),
+		bEnabled,
+		TEXT("Opens a pipe that runs commands passed as command line args to CmdLink.exe"),
+		FConsoleVariableDelegate::CreateStatic(OnCmdLinkEnabledChanged),
+		ECVF_Default);
+
+	FString CLIPipeKey = TEXT("None");
+	void OnCmdLinkKeyChanged(IConsoleVariable* Var);
+	FAutoConsoleVariableRef CVarCmdLinkKey(
+		TEXT("console.CmdLink.key"),
+		CLIPipeKey,
+		TEXT("Changes the name of the pipe used for a connection to CmdLink.exe"),
+		FConsoleVariableDelegate::CreateStatic(OnCmdLinkKeyChanged),
+		ECVF_Default);
+
+	void OnCmdLinkEnabledChanged(IConsoleVariable* Var)
+	{
+		FCmdLinkServerModule* Server = FCmdLinkServerModule::Get();
+		if (Server && bEnabled)
+		{
+			Server->Enable();
+		}
+		else
+		{
+			Server->Disable();
+		}
+	}
+
+	void OnCmdLinkKeyChanged(IConsoleVariable* Var)
+	{
+		if (FCmdLinkServerModule* Server = FCmdLinkServerModule::Get())
+		{
+			Server->OnKeyChanged(CLIPipeKey);
+		}
+	}
+	
+	extern UNREALED_API void(*GBeginAsyncCommand)(const FString&, const TArray<FString>&);
+	extern UNREALED_API void(*GEndAsyncCommand)(const FString&, const TArray<FString>&);
+}
+
+
+
+// because we're using a platform pipe in async mode and UE doesn't have coroutines, we can use this state machine
+// to change which methods are being ticked
+struct FCmdLinkServerModule::FTickStateMachine
+{
+	FTickStateMachine(int32 inCurrentTask)
+		: CurrentTask(inCurrentTask)
+	{
+		Ticker = FTSTicker::GetCoreTicker().AddTicker(FTickerDelegate::CreateRaw(this, &FTickStateMachine::Tick));
+	}
+
+	~FTickStateMachine()
+	{
+		FTSTicker::GetCoreTicker().RemoveTicker(Ticker);
+	}
+
+	void Emplace(int32 Key, TFunction<int32(float)>&& Task)
+	{
+		Tasks.Emplace(Key, Task);
+	}
+
+	int32 GetCurrentTask() const
+	{
+		return CurrentTask;
+	}
+
+private:
+	bool Tick(float dt)
+	{
+		CurrentTask = Tasks[CurrentTask](dt);
+		return Tasks.Find(CurrentTask) != nullptr;
+	}
+	
+	int32 CurrentTask = 0;
+	TMap<int32, TFunction<int32(float)>> Tasks;
+	FTSTicker::FDelegateHandle Ticker;
+};
+
+enum ETask
+{
+	BEGIN_CONNECT, // create pipe and open the connection
+	SLEEP_FOR_RECONNECT, // if BEGIN_CONNECT fails, wait a bit before trying again
+	AWAIT_CONNECT, // wait for client to be ready
+	
+	BEGIN_READ,	// begin read operation
+	AWAIT_READ, // wait for read operation to complete
+	
+	RUN_COMMAND, // run the command
+	AWAIT_COMMAND, // wait for command to complete
+	
+	REPLY, // send the next reply data in the queue
+	AWAIT_REPLY, // wait for reply to finish send
+
+
+	EXIT = -1 // stop ticking
+};
+
+void FCmdLinkServerModule::StartupModule()
+{
+	if (UE::CmdLink::bEnabled)
+	{
+		Enable();
+	}
+	
+	UE::CmdLink::GBeginAsyncCommand = [](const FString& CommandName, const TArray<FString>& Params)
+	{
+		if (FCmdLinkServerModule* Module = FCmdLinkServerModule::Get())
+		{
+			Module->BeginAsyncCommand(CommandName, Params);
+		}
+	};
+	UE::CmdLink::GEndAsyncCommand = [](const FString& CommandName, const TArray<FString>& Params)
+	{
+		if (FCmdLinkServerModule* Module = FCmdLinkServerModule::Get())
+		{
+			Module->EndAsyncCommand(CommandName, Params);
+		}
+	};
+}
+
+void FCmdLinkServerModule::ShutdownModule()
+{
+	if (UE::CmdLink::bEnabled)
+	{
+		Disable();
+	}
+}
+
+void FCmdLinkServerModule::Enable()
+{
+	bEnabled = true;
+	if (!StateMachine.IsValid())
+	{
+		InitStateMachine();
+	}
+}
+
+void FCmdLinkServerModule::Disable()
+{
+	bEnabled = false;
+	if (StateMachine.IsValid() && StateMachine->GetCurrentTask() == BEGIN_CONNECT || StateMachine->GetCurrentTask() == AWAIT_CONNECT)
+	{
+		OnPipeClosed();
+		StateMachine.Reset();
+	}
+}
+
+void FCmdLinkServerModule::OnKeyChanged(const FString& NewKey)
+{
+	if (StateMachine.IsValid() && StateMachine->GetCurrentTask() == BEGIN_CONNECT || StateMachine->GetCurrentTask() == AWAIT_CONNECT)
+	{
+		OnPipeClosed();
+		StateMachine.Reset();
+		if (bEnabled)
+		{
+			InitStateMachine();
+		}
+	}
+}
+
+void FCmdLinkServerModule::BeginAsyncCommand(const FString& CommandName, const TArray<FString>& Params)
+{
+	if (StateMachine->GetCurrentTask() != ETask::RUN_COMMAND)
+	{
+		return;
+	}
+	
+	if (CommandName != ArgV[1])
+	{
+		return;
+	}
+
+	// first two args are file path and commandName
+	if (ArgV.Num() - 2 != Params.Num())
+	{
+		return;
+	}
+
+	for (int ArgI = 0; ArgI < Params.Num(); ++ArgI)
+	{
+		if (ArgV[ArgI + 2].Data != Params[ArgI])
+		{
+			return;
+		}
+	}
+
+	bAwaitAsyncCommand = true;
+}
+
+void FCmdLinkServerModule::EndAsyncCommand(const FString& CommandName, const TArray<FString>& Params)
+{
+	if (!bAwaitAsyncCommand)
+	{
+		return;
+	}
+	
+	if (StateMachine->GetCurrentTask() != ETask::RUN_COMMAND && StateMachine->GetCurrentTask() != ETask::AWAIT_COMMAND)
+	{
+		return;
+	}
+	
+	// first two args are file path and commandName
+	if (CommandName != ArgV[1])
+	{
+		return;
+	}
+
+	if (ArgV.Num() - 2 != Params.Num())
+	{
+		return;
+	}
+
+	for (int ArgI = 0; ArgI < Params.Num(); ++ArgI)
+	{
+		if (ArgV[ArgI + 2].Data != Params[ArgI])
+		{
+			return;
+		}
+	}
+
+	bAwaitAsyncCommand = false;
+}
+
+bool FCmdLinkServerModule::Read()
+{
+	if (ArgC == ArgV.Num())
+	{
+		return NamedPipe.ReadInt32(ArgC);
+	}
+	
+	if (ArgV.Num() < ArgC)
+	{
+		// make sure ArgV is reserved to the correct size so we don't invalidate pointers
+		// as we add to the array
+		ArgV.Reserve(ArgC);
+		
+		if (NextStringSize == 0)
+		{
+			return NamedPipe.ReadInt32(NextStringSize);
+		}
+		
+		ArgV.Emplace(NextStringSize);
+		const bool Result = NamedPipe.ReadBytes(NextStringSize, ArgV.Last());
+		NextStringSize = 0;
+		
+		return Result;
+	}
+	
+	ensure(false);
+	return false;
+}
+
+bool FCmdLinkServerModule::Execute(FString& Response)
+{
+	// merge command into a single line
+	FString Command;
+	for (int Arg = 1; Arg < ArgV.Num(); ++Arg)
+	{
+		Command += ArgV[Arg].Data;
+		if (Arg < ArgV.Num() - 1)
+		{
+			Command += " ";
+		}
+	}
+	if (Command.IsEmpty())
+	{
+		// if a command wasn't sent, assume they need help
+		Command = TEXT("help");
+	}
+
+	// try running command on every command executor until one is successful (this allows us to run both engine and python commands)
+	TArray<IConsoleCommandExecutor*> CommandExecutors = IModularFeatures::Get().GetModularFeatureImplementations<IConsoleCommandExecutor>(TEXT("ConsoleCommandExecutor"));
+	for (IConsoleCommandExecutor* CommandExecutor : CommandExecutors)
+	{
+		FStringOutputDevice Ar;
+		Ar.SetAutoEmitLineTerminator(true);
+		
+		GLog->AddOutputDevice(&Ar);
+		const bool bSuccess = CommandExecutor->Exec(*Command);
+		GLog->RemoveOutputDevice(&Ar);
+		
+		if (bSuccess)
+		{
+			Response = Ar;
+			return true;
+		}
+	}
+	
+	return false;
+}
+
+void FCmdLinkServerModule::InitStateMachine()
+{
+	// The CmdLinkServer expects messages from the client in the following format:
+	// 1. int32: ArgC
+	// repeat lines 2 and 3 ArgC times
+	// 2. int32: ArgLen
+	// 3. char[Arglen]: Arg (includes null terminating character)
+
+	// the CmdLink client expects a response from the server in the following format:
+	// 1. int32: responseLen
+	// 2. char[responseLen] response (includes null terminating character)
+	
+	StateMachine = MakeShared<FTickStateMachine>(BEGIN_CONNECT);
+	
+	StateMachine->Emplace(BEGIN_CONNECT, [this](float dt)->int32
+	{
+		// if there's a key specified, include it in the name
+		const FString PipeName = UE::CmdLink::CLIPipeKey == TEXT("None") ? TEXT("UnrealEngine-CLI") : TEXT("UnrealEngine-CLI-") + UE::CmdLink::CLIPipeKey;
+		// Create the output pipe as a server...
+		if (!NamedPipe.Create(FString::Printf(TEXT("\\\\.\\pipe\\%s"), *PipeName), true, true))
+		{
+			UE_LOG(LogEngine, Error, TEXT("Failed to create CLI pipe"));
+			return SLEEP_FOR_RECONNECT; // wait 5 seconds and try again
+		}
+		if (!NamedPipe.OpenConnection())
+		{
+			UE_LOG(LogEngine, Error, TEXT("Failed to open a connection on the CLI pipe"));
+			return SLEEP_FOR_RECONNECT; // wait 5 seconds and try again
+		}
+		return AWAIT_CONNECT;
+	});
+	
+	StateMachine->Emplace(SLEEP_FOR_RECONNECT, [this](float dt)->int32
+	{
+		if (SleepConnectTimer <= 0.f)
+		{
+			SleepConnectTimer = 5.f; // wait for 5 seconds
+		}
+		SleepConnectTimer -= dt;
+		if (SleepConnectTimer <= 0.f)
+		{
+			return (bEnabled) ? BEGIN_CONNECT : EXIT;
+		}
+		return SLEEP_FOR_RECONNECT;
+	});
+	
+	StateMachine->Emplace(AWAIT_CONNECT, [this](float dt)->int32
+	{
+		if (!NamedPipe.UpdateAsyncStatus())
+		{
+			return OnPipeClosed();
+		}
+		
+		if (NamedPipe.IsReadyForRW())
+		{
+			return BEGIN_READ;
+		}
+		return AWAIT_CONNECT;
+	});
+	
+	StateMachine->Emplace(BEGIN_READ, [this](float dt)->int32
+	{
+		if (!NamedPipe.UpdateAsyncStatus())
+		{
+			return OnPipeClosed();
+		}
+		
+		if (NamedPipe.IsReadyForRW())
+		{
+			if (!Read())
+			{
+				return OnPipeClosed();
+			}
+			return AWAIT_READ;
+		}
+		return BEGIN_READ;
+	});
+	
+	StateMachine->Emplace(AWAIT_READ, [this](float dt)->int32
+	{
+		if (!NamedPipe.UpdateAsyncStatus())
+		{
+			return OnPipeClosed();
+		}
+		if (NamedPipe.IsReadyForRW())
+		{
+			if (ArgV.Num() < ArgC)
+			{
+				// we need to read more data
+				return BEGIN_READ;
+			}
+			// success
+			return RUN_COMMAND;
+		}
+		return AWAIT_READ;
+	});
+	
+	StateMachine->Emplace(RUN_COMMAND, [this](float dt)->int32
+	{
+		FString Response;
+		const bool bSuccess = Execute(Response);
+		const TStringConversion Result = StringCast<ANSICHAR>(*Response);
+		// queue up response to the client
+		if (Result.Length() == 0 && bSuccess == false)
+		{
+			const ANSICHAR Error[] = "Command failed or unrecognized";
+			const int32 Size = TCString<ANSICHAR>::Strlen(Error) + 1;
+			PendingReply.Enqueue(TArray(reinterpret_cast<const uint8*>(&Size), sizeof(Size)));
+			PendingReply.Enqueue(TArray(reinterpret_cast<const uint8*>(Error), Size));
+		}
+		else
+		{
+			const int32 Size = Result.Length() + 1;
+			PendingReply.Enqueue(TArray(reinterpret_cast<const uint8*>(&Size), sizeof(Size)));
+			PendingReply.Enqueue(TArray(reinterpret_cast<const uint8*>(Result.Get()), Size));
+		}
+		
+		if (bAwaitAsyncCommand)
+		{
+			return AWAIT_COMMAND;
+		}
+		return REPLY;
+	});
+	
+	StateMachine->Emplace(AWAIT_COMMAND, [this](float dt)->int32
+	{
+		if (bAwaitAsyncCommand)
+		{
+			return AWAIT_COMMAND;
+		}
+		return REPLY;
+	});
+	
+	StateMachine->Emplace(REPLY, [this](float dt)->int32
+	{
+		
+		if (!NamedPipe.UpdateAsyncStatus())
+		{
+			return OnPipeClosed();
+		}
+		
+		if (NamedPipe.IsReadyForRW())
+		{
+			TArray<uint8> Buffer;
+			PendingReply.Dequeue(Buffer);
+			if (!NamedPipe.WriteBytes(Buffer.Num(), Buffer.GetData()))
+			{
+				return OnPipeClosed();
+			}
+			return AWAIT_REPLY;
+		}
+		return REPLY;
+	});
+	
+	StateMachine->Emplace(AWAIT_REPLY, [this](float dt)->int32
+	{
+		if (!NamedPipe.UpdateAsyncStatus())
+		{
+			return OnPipeClosed();
+		}
+		if (NamedPipe.IsReadyForRW())
+		{
+			if (PendingReply.IsEmpty())
+			{
+				// success! loop back and wait for the next message
+				ArgV.Reset();
+				ArgC = 0;
+				NextStringSize = 0;
+				return BEGIN_READ;
+			}
+			return REPLY;
+		}
+		return AWAIT_REPLY;
+	});
+}
+
+int32 FCmdLinkServerModule::OnPipeClosed()
+{
+	// pipe connection broke. reconnect
+	NamedPipe.Destroy();
+	ArgV.Reset();
+	ArgC = 0;
+	NextStringSize = 0;
+	PendingReply.Empty();
+	return (bEnabled) ? BEGIN_CONNECT : EXIT;
+}
+
+#undef LOCTEXT_NAMESPACE
+	
+IMPLEMENT_MODULE(FCmdLinkServerModule, CmdLinkServer)
\ No newline at end of file
diff --git a/Engine/Plugins/CmdLinkServer/Source/CmdLinkServer/Public/CmdLinkServer.h b/Engine/Plugins/CmdLinkServer/Source/CmdLinkServer/Public/CmdLinkServer.h
new file mode 100644
--- /dev/null
+++ b/Engine/Plugins/CmdLinkServer/Source/CmdLinkServer/Public/CmdLinkServer.h
@@ -0,0 +1,57 @@
+﻿// Copyright Epic Games, Inc. All Rights Reserved.
+
+#pragma once
+
+#include "CoreMinimal.h"
+#include "Modules/ModuleManager.h"
+#include "Windows/WindowsPlatformNamedPipe.h"
+#include "Containers/Queue.h"
+
+class FCmdLinkServerModule : public IModuleInterface
+{
+public:
+	/** IModuleInterface implementation */
+	virtual void StartupModule() override;
+	virtual void ShutdownModule() override;
+	void Enable();
+	void Disable();
+	void OnKeyChanged(const FString& NewKey);
+	static FCmdLinkServerModule* Get()
+	{
+		return FModuleManager::GetModulePtr<FCmdLinkServerModule>("CmdLinkServer");
+	}
+
+	// call this from a command that will take more than one frame to complete
+	// this forces the CmdLink to wait until EndAsyncCommand is called before shutting down
+	CMDLINKSERVER_API void BeginAsyncCommand(const FString& CommandName, const TArray<FString>& Params);
+	
+	// call this at the end of a command that took more than one frame to complete
+	CMDLINKSERVER_API void EndAsyncCommand(const FString& CommandName, const TArray<FString>& Params);
+private:
+	bool Read();
+	bool Execute(FString& Response);
+	void InitStateMachine();
+	int32 OnPipeClosed();
+	
+	// scoped ANSICHAR c-string buffer that auto-casts to an FString
+	struct FAnsiStringBuffer
+	{
+		FAnsiStringBuffer(int32 Size) : Data(new ANSICHAR[Size]{0}) {}
+		~FAnsiStringBuffer() {delete [] Data;}
+		operator FString() const {return FString(Data);}
+		operator ANSICHAR*() const {return Data;}
+		
+		ANSICHAR* Data;
+	};
+	
+	struct FTickStateMachine;
+	bool bEnabled = false;
+	FPlatformNamedPipe NamedPipe;
+	TSharedPtr<FTickStateMachine> StateMachine;
+	int32 ArgC = 0;
+	TArray<FAnsiStringBuffer> ArgV;
+	int32 NextStringSize = 0;
+	bool bAwaitAsyncCommand = false;
+	TQueue<TArray<uint8>> PendingReply;
+	float SleepConnectTimer = -1.f;
+};
diff --git a/Engine/Source/Editor/UnrealEd/Private/DiffUtils.cpp b/Engine/Source/Editor/UnrealEd/Private/DiffUtils.cpp
--- a/Engine/Source/Editor/UnrealEd/Private/DiffUtils.cpp
+++ b/Engine/Source/Editor/UnrealEd/Private/DiffUtils.cpp
@@ -2,9 +2,11 @@
 
 #include "DiffUtils.h"
 
+#include "AssetDefinitionRegistry.h"
 #include "Components/ActorComponent.h"
 #include "Containers/BitArray.h"
 #include "EditorCategoryUtils.h"
+#include "IAssetTools.h"
 #include "Engine/Blueprint.h"
 #include "IAssetTypeActions.h"
 #include "ISourceControlModule.h"
@@ -30,19 +32,23 @@
 #include "UObject/Object.h"
 #include "UObject/ObjectMacros.h"
 #include "UObject/PropertyPortFlags.h"
+#include "UObject/SavePackage.h"
 #include "Widgets/DeclarativeSyntaxSupport.h"
 #include "Widgets/Images/SImage.h"
 #include "Widgets/Text/STextBlock.h"
 #include "Widgets/Views/STableRow.h"
 #include "Widgets/Views/STableViewBase.h"
 #include "Widgets/Views/STreeView.h"
+#include "HAL/PlatformFileManager.h"
+#include "UnrealEngine.h"
+#include "UObject/Linker.h"
 
 class ITableRow;
 class SWidget;
 
 namespace UEDiffUtils_Private
 {
-	FProperty* Resolve( const UStruct* Class, FName PropertyName )
+	static FProperty* Resolve( const UStruct* Class, FName PropertyName )
 	{
 		if(Class == nullptr )
 		{
@@ -60,10 +66,206 @@ namespace UEDiffUtils_Private
 		return nullptr;
 	}
 
-	FPropertySoftPathSet GetPropertyNameSet(const UStruct* ForStruct)
+	static FPropertySoftPathSet GetPropertyNameSet(const UStruct* ForStruct)
 	{
 		return FPropertySoftPathSet(DiffUtils::GetVisiblePropertiesInOrderDeclared(ForStruct));
 	}
+	
+	static const FString DiffSyntaxHelp = TEXT("format: 'diff <lhs> <rhs>");
+	static const FString MergeSyntaxHelp = TEXT("format: 'merge <remote> <local> <base> [-o out_path]' or 'merge <local> [-o out_path]'");
+	static void RunDiffCommand(const TArray<FString>& Args);
+	static void RunMergeCommand(const TArray<FString>& Args);
+
+	FAutoConsoleCommand DiffConsoleCommand(
+		TEXT("merge"),
+		*FString::Format(TEXT("Either merge three assets or a single conflicted asset.\n{0}"), {MergeSyntaxHelp}),
+		FConsoleCommandWithArgsDelegate::CreateStatic(&RunMergeCommand),
+		ECVF_Default
+	);
+
+	FAutoConsoleCommand MergeConsoleCommand(
+		TEXT("diff"),
+		*FString::Format(TEXT("diff two assets against one another.\n{0}"), {DiffSyntaxHelp}),
+		FConsoleCommandWithArgsDelegate::CreateStatic(&RunDiffCommand),
+		ECVF_Default
+	);
+}
+
+static UObject* LoadAssetFromExternalPath(FString Path)
+{
+	FPackagePath PackagePath;
+	if (!FPackagePath::TryFromPackageName(Path, PackagePath))
+	{
+		// copy to the temp directory so it can be loaded properly
+		FString File = FPaths::GetBaseFilename(Path) + TEXT("-");
+		for (const ANSICHAR Char : "#(){}[].")
+		{
+			File.ReplaceCharInline(Char, '-');
+		}
+		const FString Extension = TEXT(".") + FPaths::GetExtension(Path);
+		const FString SourcePath = Path;
+		FPlatformFileManager::Get().GetPlatformFile().CreateDirectory(*FPaths::DiffDir());
+		Path = FPaths::CreateTempFilename(*FPaths::DiffDir(), *File, *Extension);
+		Path = FPaths::ConvertRelativePathToFull(Path);
+		if (!FPlatformFileManager::Get().GetPlatformFile().CopyFile(*Path, *SourcePath))
+		{
+			UE_LOG(LogEngine, Display, TEXT("Failed to Copy %s"), *SourcePath);
+			return nullptr;
+		}
+		// load the temp package
+		PackagePath = FPackagePath::FromLocalPath(Path);
+	}
+	if (PackagePath.IsEmpty())
+	{
+		UE_LOG(LogEngine, Display, TEXT("Invalid Path: %s"), *Path);
+		return nullptr;
+	}
+	if (const UPackage* TempPackage = DiffUtils::LoadPackageForDiff(PackagePath, {}))
+	{
+		if (UObject* Object = TempPackage->FindAssetInPackage())
+		{
+			return Object;
+		}
+	}
+	UE_LOG(LogEngine, Display, TEXT("Failed to load: %s"), *Path);
+	return nullptr;
+}
+
+static void UEDiffUtils_Private::RunDiffCommand(const TArray<FString>& Args)
+{
+	if (Args.Num() != 2)
+	{
+		UE_LOG(LogEngine, Display, TEXT("%s"), *DiffSyntaxHelp);
+		return;
+	}
+	
+	UObject* LHS = LoadAssetFromExternalPath(Args[0]);
+	UObject* RHS = LoadAssetFromExternalPath(Args[1]);
+	if (LHS && RHS)
+	{
+		IAssetTools::Get().DiffAssets(LHS, RHS, {}, {});
+	}
+}
+
+namespace UE::CmdLink
+{
+	// CmdLinkServerModule will set these methods if loaded.
+	// they're used by the merge command because we need to keep the CmdLink client running until the user closes the merge window and saves the output
+	UNREALED_API void(*GBeginAsyncCommand)(const FString&, const TArray<FString>&) = [](const FString&, const TArray<FString>&){};
+	UNREALED_API void(*GEndAsyncCommand)(const FString&, const TArray<FString>&) = [](const FString&, const TArray<FString>&){};
+}
+
+static void UEDiffUtils_Private::RunMergeCommand(const TArray<FString>& Args)
+{
+	UObject* Local = nullptr;
+	UObject* Base = nullptr;
+	UObject* Remote = nullptr;
+	FString OutDirectory;
+	bool bThreeWayMerge = false;
+	bool bInvalidSyntax = false;
+	switch (Args.Num())
+	{
+	case 1: // merge <local>
+		Local = LoadAssetFromExternalPath(Args[0]);
+		bThreeWayMerge = false;
+		break;
+	case 3:
+		if (Args[1] == TEXT("-o")) // merge <local> -o <output_file>
+		{
+			Local = LoadAssetFromExternalPath(Args[0]);
+			OutDirectory = Args[2];
+			bThreeWayMerge = false;
+		}
+		else // merge <local> <base> <remote>
+		{
+			Remote = LoadAssetFromExternalPath(Args[0]);
+			Local = LoadAssetFromExternalPath(Args[1]);
+			Base = LoadAssetFromExternalPath(Args[2]);
+			bThreeWayMerge = true;
+		}
+		break;
+	case 5: // merge <local> <base> <remote> -o <output_file>
+		if (Args[3] == TEXT("-o"))
+		{
+			Remote = LoadAssetFromExternalPath(Args[0]);
+			Local = LoadAssetFromExternalPath(Args[1]);
+			Base = LoadAssetFromExternalPath(Args[2]);
+			OutDirectory = Args[4];
+			bThreeWayMerge = true;
+			break;
+		}
+
+		// 5 parameters requires output file at the end
+		bInvalidSyntax = true;
+		break;
+	default:
+		// unsupported parameter count
+		bInvalidSyntax = true;
+		break;
+	}
+
+	if (bInvalidSyntax)
+	{
+		// invalid syntax. display help.
+		UE_LOG(LogEngine, Display, TEXT("%s"), *MergeSyntaxHelp);
+		return;
+	}
+
+	const FOnAssetMergeResolved ResolutionCallback = FOnAssetMergeResolved::CreateLambda([Args, Local, OutDirectory](const FAssetMergeResults& Results)
+	{
+		if (!OutDirectory.IsEmpty() && Results.Result == EAssetMergeResult::Completed)
+		{
+			// save a copy of the asset to the output directory
+			
+			FSavePackageArgs SaveArgs;
+			SaveArgs.TopLevelFlags = RF_Standalone;
+			SaveArgs.Error = GLog;
+			UPackage::SavePackage(Results.MergedPackage, Local, *OutDirectory, SaveArgs);
+			ResetLoaders(Results.MergedPackage);
+		}
+		UE::CmdLink::GEndAsyncCommand(TEXT("merge"), Args);
+	});
+
+	if (bThreeWayMerge)
+	{
+		FAssetManualMergeArgs MergeArgs;
+		MergeArgs.LocalAsset = Local;
+		MergeArgs.BaseAsset = Base;
+		MergeArgs.RemoteAsset = Remote;
+		MergeArgs.ResolutionCallback = ResolutionCallback;
+		MergeArgs.Flags = MF_NONE;
+
+		if (MergeArgs.LocalAsset && MergeArgs.BaseAsset && MergeArgs.RemoteAsset)
+		{
+			const UAssetDefinition* AssetDefinition = UAssetDefinitionRegistry::Get()->GetAssetDefinitionForClass(Local->GetClass());
+			if (!AssetDefinition->CanMerge())
+			{
+				UE_LOG(LogEngine, Error, TEXT("%s of class type %s does not support merging"), *Local->GetName(), *Local->GetClass()->GetName());
+				return;
+			}
+			UE::CmdLink::GBeginAsyncCommand(TEXT("merge"), Args);
+			AssetDefinition->Merge(MergeArgs);
+		}
+	}
+	else
+	{
+		FAssetAutomaticMergeArgs MergeArgs;
+		MergeArgs.LocalAsset = Local;
+		MergeArgs.ResolutionCallback = ResolutionCallback;
+		MergeArgs.Flags = MF_NONE;
+		
+		if (MergeArgs.LocalAsset)
+		{
+			const UAssetDefinition* AssetDefinition = UAssetDefinitionRegistry::Get()->GetAssetDefinitionForClass(Local->GetClass());
+			if (!AssetDefinition->CanMerge())
+			{
+				UE_LOG(LogEngine, Error, TEXT("%s does not support merging"), *Local->GetName());
+				return;
+			}
+			UE::CmdLink::GBeginAsyncCommand(TEXT("merge"), Args);
+			AssetDefinition->Merge(MergeArgs);
+		}
+	}
 }
 
 FPropertySoftPath::FPropertySoftPath()
@@ -847,6 +1049,12 @@ TArray<FPropertyPath> DiffUtils::ResolveAll(const UObject* Object, const TArray<
 
 UPackage* DiffUtils::LoadPackageForDiff(const FPackagePath& InTempPackagePath, const FPackagePath& InOriginalPackagePath)
 {
+	// if this is a local asset, load it normally
+	if (!FPackageName::IsTempPackage(InTempPackagePath.GetPackageName()))
+	{
+		return LoadPackage(nullptr, *InTempPackagePath.GetPackageName(), LOAD_None);
+	}
+	
 	// set up instancing context
 	FLinkerInstancingContext Context;
 	if (!InOriginalPackagePath.GetLocalFullPath().IsEmpty())
diff --git a/Engine/Source/Programs/CmdLink/CmdLink.Build.cs b/Engine/Source/Programs/CmdLink/CmdLink.Build.cs
new file mode 100644
--- /dev/null
+++ b/Engine/Source/Programs/CmdLink/CmdLink.Build.cs
@@ -0,0 +1,22 @@
+﻿// Copyright Epic Games, Inc. All Rights Reserved.
+
+using UnrealBuildTool;
+
+public class CmdLink : ModuleRules
+{
+    public CmdLink(ReadOnlyTargetRules Target) : base(Target)
+    {
+        PrivateDependencyModuleNames.AddRange(
+            new string[]
+            {
+	            "Core",
+            }
+        );
+        
+        PrivateIncludePathModuleNames.AddRange(
+	        new string[] {
+		        "Launch",
+		        "TargetPlatform",
+	        });
+    }
+}
\ No newline at end of file
diff --git a/Engine/Source/Programs/CmdLink/CmdLink.Target.cs b/Engine/Source/Programs/CmdLink/CmdLink.Target.cs
new file mode 100644
--- /dev/null
+++ b/Engine/Source/Programs/CmdLink/CmdLink.Target.cs
@@ -0,0 +1,35 @@
+// Copyright Epic Games, Inc. All Rights Reserved.
+
+using System.Collections.Generic;
+using System.IO;
+using UnrealBuildTool;
+
+[SupportedPlatforms("Win64")]
+public class CmdLinkTarget : TargetRules
+{
+	public CmdLinkTarget(TargetInfo Target) : base(Target)
+	{
+		Type = TargetType.Program;
+		LinkType = TargetLinkType.Monolithic;
+		LaunchModuleName = "CmdLink";
+		bLegalToDistributeBinary = true;
+
+		// Turn off various third party features we don't need
+
+		// Lean and mean
+		bBuildDeveloperTools = false;
+
+		// CmdLink isn't localized, so doesn't need ICU
+		bCompileICU = false;
+
+		// Currently this app is not linking against the engine, so we'll compile out references from Core to the rest of the engine
+		bCompileAgainstEngine = false;
+		bCompileAgainstCoreUObject = false;
+		bCompileAgainstApplicationCore = false;
+		bBuildWithEditorOnlyData = false;
+		bCompileCEF3 = false;
+
+		// CmdLink is a console application, not a Windows app (sets entry point to main(), instead of WinMain())
+		bIsBuildingConsoleApplication = true;
+	}
+}
diff --git a/Engine/Source/Programs/CmdLink/Private/CmdLink.cpp b/Engine/Source/Programs/CmdLink/Private/CmdLink.cpp
new file mode 100644
--- /dev/null
+++ b/Engine/Source/Programs/CmdLink/Private/CmdLink.cpp
@@ -0,0 +1,93 @@
+﻿// Copyright Epic Games, Inc. All Rights Reserved.
+
+#include <iostream>
+
+#include "Core.h"
+#include "Windows/WindowsPlatformNamedPipe.h"
+
+enum EErrorMessage
+{
+	SUCCESS = 0,
+	CONNECT_FAILED,
+	WRITE_FAILED,
+	READ_FAILED
+};
+
+EErrorMessage PrintError(EErrorMessage Error)
+{
+	switch(Error)
+	{
+	case CONNECT_FAILED:
+		printf("Failed to connect to UE\n");
+		break;
+	case WRITE_FAILED:
+		printf("Failed to send command to UE\n");
+		break;
+	case READ_FAILED:
+		printf("Failed to read response from UE\n");
+		break;
+	}
+	return Error;
+}
+
+#define TRY(code, error) if(!(code)) return PrintError(error)
+
+int32 main(int32 ArgC, ANSICHAR* ArgV[])
+{
+	FString CLIPipeKey;
+	if (ArgC >= 2 && FParse::Value(StringCast<TCHAR>(ArgV[1]).Get(),TEXT("--LinkKey="),CLIPipeKey))
+	{
+		// don't send this argument to the server. it's meant for Client only
+		--ArgC;
+		++ArgV;
+	}
+	
+	// The CLIServer expects messages from the client in the following format:
+	// 1. int32: ArgC
+	// repeat lines 2 and 3 ArgC times
+	// 2. int32: ArgLen
+	// 3. char[Arglen]: Arg (includes null terminating character)
+
+	// the CmdLink client expects a response from the server in the following format:
+	// 1. int32: responseLen
+	// 2. char[responseLen] response (includes null terminating character)
+	
+	const FString PipeName = CLIPipeKey.IsEmpty() ? TEXT("UnrealEngine-CLI") : TEXT("UnrealEngine-CLI-") + CLIPipeKey;
+	FPlatformNamedPipe NamedPipe;
+	TRY(NamedPipe.Create(FString::Printf(TEXT("\\\\.\\pipe\\%s"), *PipeName), false, false), CONNECT_FAILED);
+
+	// Send ArgC
+	TRY(NamedPipe.WriteInt32(ArgC), WRITE_FAILED);
+
+	// send file path of caller
+	{
+		FString LaunchDir = FPlatformProcess::GetCurrentWorkingDirectory();
+		FPaths::NormalizeDirectoryName(LaunchDir);
+		const TStringConversion AnsiStr = StringCast<ANSICHAR>(*LaunchDir);
+		const int32 Size = AnsiStr.Length() + 1;
+		TRY(NamedPipe.WriteInt32(Size), WRITE_FAILED);
+		TRY(NamedPipe.WriteBytes(Size, AnsiStr.Get()), WRITE_FAILED);
+	}
+
+	// Send ArgV (skipping filename)
+	for(int i = 1; i < ArgC; ++i)
+	{
+		const int32 Size = TCString<ANSICHAR>::Strlen(ArgV[i]) + 1;
+		TRY(NamedPipe.WriteInt32(Size), WRITE_FAILED);
+		TRY(NamedPipe.WriteBytes(Size, ArgV[i]), WRITE_FAILED);
+	}
+
+	// Receive Response string
+	int32 Size;
+	TRY(NamedPipe.ReadInt32(Size), READ_FAILED);
+	ANSICHAR* Buffer = new ANSICHAR[Size];
+	TRY(NamedPipe.ReadBytes(Size, Buffer), READ_FAILED);
+
+	// Print Response
+	check(Buffer[Size - 1] == '\0');
+	printf(Buffer);
+
+	delete [] Buffer;
+	NamedPipe.Destroy();
+	return SUCCESS;
+}
diff --git a/Engine/Source/Runtime/Core/Private/Windows/WindowsPlatformNamedPipe.cpp b/Engine/Source/Runtime/Core/Private/Windows/WindowsPlatformNamedPipe.cpp
--- a/Engine/Source/Runtime/Core/Private/Windows/WindowsPlatformNamedPipe.cpp
+++ b/Engine/Source/Runtime/Core/Private/Windows/WindowsPlatformNamedPipe.cpp
@@ -109,7 +109,6 @@ bool FWindowsPlatformNamedPipe::Destroy()
 		case State_WaitingForRW:
 			// TODO: Cancel IO and wait
 			bFlushBuffers = true;
-			check(0);
 			break;
 
 		case State_Uninitialized:
@@ -118,7 +117,6 @@ bool FWindowsPlatformNamedPipe::Destroy()
 
 		case State_Connecting:
 			// TODO: Cancel IO
-			check(0);
 			break;
 
 		case State_ErrorPipeClosedUnexpectedly:
