Fix bad operator precedence

diff --git a/Engine/Plugins/CmdLinkServer/Source/CmdLinkServer/Private/CmdLinkServer.cpp b/Engine/Plugins/CmdLinkServer/Source/CmdLinkServer/Private/CmdLinkServer.cpp
--- a/Engine/Plugins/CmdLinkServer/Source/CmdLinkServer/Private/CmdLinkServer.cpp
+++ b/Engine/Plugins/CmdLinkServer/Source/CmdLinkServer/Private/CmdLinkServer.cpp
@@ -156,7 +156,7 @@ void FCmdLinkServerModule::Enable()
 void FCmdLinkServerModule::Disable()
 {
 	bEnabled = false;
-	if (StateMachine.IsValid() && StateMachine->GetCurrentTask() == BEGIN_CONNECT || StateMachine->GetCurrentTask() == AWAIT_CONNECT)
+	if (StateMachine.IsValid() && (StateMachine->GetCurrentTask() == BEGIN_CONNECT || StateMachine->GetCurrentTask() == AWAIT_CONNECT))
 	{
 		OnPipeClosed();
 		StateMachine.Reset();
@@ -165,7 +165,7 @@ void FCmdLinkServerModule::Disable()
 
 void FCmdLinkServerModule::OnKeyChanged(const FString& NewKey)
 {
-	if (StateMachine.IsValid() && StateMachine->GetCurrentTask() == BEGIN_CONNECT || StateMachine->GetCurrentTask() == AWAIT_CONNECT)
+	if (StateMachine.IsValid() && (StateMachine->GetCurrentTask() == BEGIN_CONNECT || StateMachine->GetCurrentTask() == AWAIT_CONNECT))
 	{
 		OnPipeClosed();
 		StateMachine.Reset();
