Fix potential null dereference likely causing plugin test failure
#rnx
#horde 329879

diff --git a/Engine/Plugins/CmdLinkServer/Source/CmdLinkServer/Private/CmdLinkServer.cpp b/Engine/Plugins/CmdLinkServer/Source/CmdLinkServer/Private/CmdLinkServer.cpp
--- a/Engine/Plugins/CmdLinkServer/Source/CmdLinkServer/Private/CmdLinkServer.cpp
+++ b/Engine/Plugins/CmdLinkServer/Source/CmdLinkServer/Private/CmdLinkServer.cpp
@@ -32,14 +32,16 @@ namespace UE::CmdLink
 
 	void OnCmdLinkEnabledChanged(IConsoleVariable* Var)
 	{
-		FCmdLinkServerModule* Server = FCmdLinkServerModule::Get();
-		if (Server && bEnabled)
-		{
-			Server->Enable();
-		}
-		else
+		if (FCmdLinkServerModule* Server = FCmdLinkServerModule::Get())
 		{
-			Server->Disable();
+			if (bEnabled)
+			{
+				Server->Enable();
+			}
+			else
+			{
+				Server->Disable();
+			}
 		}
 	}
 
