Disable cmdlink on build machines unless -cmdlink is passed on the command line to avoid conflicts over pipe.

diff --git a/Engine/Plugins/CmdLinkServer/Source/CmdLinkServer/Private/CmdLinkServer.cpp b/Engine/Plugins/CmdLinkServer/Source/CmdLinkServer/Private/CmdLinkServer.cpp
--- a/Engine/Plugins/CmdLinkServer/Source/CmdLinkServer/Private/CmdLinkServer.cpp
+++ b/Engine/Plugins/CmdLinkServer/Source/CmdLinkServer/Private/CmdLinkServer.cpp
@@ -18,7 +18,6 @@ namespace UE::CmdLink
 		TEXT("console.CmdLink.enable"),
 		bEnabled,
 		TEXT("Opens a pipe that runs commands passed as command line args to CmdLink.exe"),
-		FConsoleVariableDelegate::CreateStatic(OnCmdLinkEnabledChanged),
 		ECVF_Default);
 
 	FString CLIPipeKey = TEXT("None");
@@ -117,11 +116,23 @@ enum ETask
 
 void FCmdLinkServerModule::StartupModule()
 {
+	// Disable by default on build machines to prevent conflicts over pipe name
+	if (GIsBuildMachine)
+	{
+		UE::CmdLink::CVarEnableCmdLink->Set(false, ECVF_SetByCode);
+	}
+	if (FParse::Param(FCommandLine::Get(), TEXT("cmdlink")))
+	{
+		UE::CmdLink::CVarEnableCmdLink->Set(true, ECVF_SetByCode);
+	}
+
 	if (UE::CmdLink::bEnabled)
 	{
 		Enable();
 	}
 	
+	UE::CmdLink::CVarEnableCmdLink->SetOnChangedCallback(FConsoleVariableDelegate::CreateStatic(UE::CmdLink::OnCmdLinkEnabledChanged));
+	
 	UE::CmdLink::GBeginAsyncCommand = [](const FString& CommandName, const TArray<FString>& Params)
 	{
 		if (FCmdLinkServerModule* Module = FCmdLinkServerModule::Get())
