[BugFix] Silenced noisy error in CmdLinkServer plugin

diff --git a/Engine/Plugins/CmdLinkServer/Source/CmdLinkServer/Private/CmdLinkServer.cpp b/Engine/Plugins/CmdLinkServer/Source/CmdLinkServer/Private/CmdLinkServer.cpp
--- a/Engine/Plugins/CmdLinkServer/Source/CmdLinkServer/Private/CmdLinkServer.cpp
+++ b/Engine/Plugins/CmdLinkServer/Source/CmdLinkServer/Private/CmdLinkServer.cpp
@@ -341,7 +341,6 @@ void FCmdLinkServerModule::InitStateMachine()
 		// Create the output pipe as a server...
 		if (!NamedPipe.Create(FString::Printf(TEXT("\\\\.\\pipe\\%s"), *PipeName), true, true))
 		{
-			UE_LOG(LogEngine, Error, TEXT("Failed to create CLI pipe"));
 			return SLEEP_FOR_RECONNECT; // wait 5 seconds and try again
 		}
 		if (!NamedPipe.OpenConnection())
